ICVFX: Added CameraCalibration as a dependency convenience.

#rnx

diff --git a/Engine/Plugins/VirtualProduction/ICVFX/ICVFX.uplugin b/Engine/Plugins/VirtualProduction/ICVFX/ICVFX.uplugin
--- a/Engine/Plugins/VirtualProduction/ICVFX/ICVFX.uplugin
+++ b/Engine/Plugins/VirtualProduction/ICVFX/ICVFX.uplugin
@@ -15,35 +15,27 @@
 	"Installed": false,
 	"Plugins": [
 		{
-			"Name": "nDisplay",
+			"Name": "AjaMedia",
 			"Enabled": true
 		},
 		{
-			"Name": "VirtualProductionUtilities",
+			"Name": "BlackmagicMedia",
 			"Enabled": true
 		},
 		{
-			"Name": "MediaFrameworkUtilities",
+			"Name": "CameraCalibration",
 			"Enabled": true
 		},
 		{
-			"Name": "MultiUserClient",
+			"Name": "Composure",
 			"Enabled": true
 		},
 		{
-			"Name": "MultiUserTakes",
+			"Name": "ColorCorrectRegions",
 			"Enabled": true
 		},
 		{
-			"Name": "RemoteControl",
+			"Name": "GPULightmass",
 			"Enabled": true
 		},
 		{
-			"Name": "RemoteControlWebInterface",
-			"Enabled": true
-		},
-		{
-			"Name": "Switchboard",
-			"Enabled": true
-		},
-		{
@@ -67,19 +59,23 @@
 			"Enabled": true
 		},
 		{
-			"Name": "TimedDataMonitor",
+			"Name": "MediaFrameworkUtilities",
 			"Enabled": true
 		},
 		{
-			"Name": "StageMonitoring",
+			"Name": "MultiUserClient",
 			"Enabled": true
 		},
 		{
-			"Name": "Takes",
+			"Name": "MultiUserTakes",
 			"Enabled": true
 		},
 		{
-			"Name": "SequencerScripting",
+			"Name": "nDisplay",
 			"Enabled": true
 		},
 		{
+			"Name": "OSC",
+			"Enabled": true
+		},
+		{
@@ -87,27 +83,35 @@
 			"Enabled": true
 		},
 		{
-			"Name": "Composure",
+			"Name": "RemoteControl",
 			"Enabled": true
 		},
 		{
-			"Name": "ColorCorrectRegions",
+			"Name": "RemoteControlWebInterface",
 			"Enabled": true
 		},
 		{
-			"Name": "OSC",
+			"Name": "StageMonitoring",
 			"Enabled": true
 		},
 		{
-			"Name": "AjaMedia",
+			"Name": "Switchboard",
 			"Enabled": true
 		},
 		{
-			"Name": "BlackmagicMedia",
+			"Name": "Takes",
 			"Enabled": true
 		},
 		{
-			"Name": "GPULightmass",
+			"Name": "TimedDataMonitor",
 			"Enabled": true
+		},
+		{
+			"Name": "SequencerScripting",
+			"Enabled": true
+		},
+		{
+			"Name": "VirtualProductionUtilities",
+			"Enabled": true
 		}
 	]
