icvfxplugin: Added cvars tool as a dependency

diff --git a/Engine/Plugins/VirtualProduction/ICVFX/ICVFX.uplugin b/Engine/Plugins/VirtualProduction/ICVFX/ICVFX.uplugin
--- a/Engine/Plugins/VirtualProduction/ICVFX/ICVFX.uplugin
+++ b/Engine/Plugins/VirtualProduction/ICVFX/ICVFX.uplugin
@@ -113,6 +113,10 @@
 		{
 			"Name": "VirtualProductionUtilities",
 			"Enabled": true
+		},
+		{
+			"Name": "ConsoleVariables",
+			"Enabled": true
 		}
 	]
 }
\ No newline at end of file
