ICVFX: Added dependency on EpicStageApp.

diff --git a/Engine/Plugins/VirtualProduction/ICVFX/ICVFX.uplugin b/Engine/Plugins/VirtualProduction/ICVFX/ICVFX.uplugin
--- a/Engine/Plugins/VirtualProduction/ICVFX/ICVFX.uplugin
+++ b/Engine/Plugins/VirtualProduction/ICVFX/ICVFX.uplugin
@@ -117,6 +117,10 @@
 		{
 			"Name": "ConsoleVariables",
 			"Enabled": true
+		},
+		{
+			"Name": "EpicStageApp",
+			"Enabled": true
 		}
 	]
 }
\ No newline at end of file
