- Enabling rivermax in ICVFX plugin

diff --git a/Engine/Plugins/VirtualProduction/ICVFX/ICVFX.uplugin b/Engine/Plugins/VirtualProduction/ICVFX/ICVFX.uplugin
--- a/Engine/Plugins/VirtualProduction/ICVFX/ICVFX.uplugin
+++ b/Engine/Plugins/VirtualProduction/ICVFX/ICVFX.uplugin
@@ -121,6 +121,10 @@
 		{
 			"Name": "EpicStageApp",
 			"Enabled": true
+		},
+		{
+			"Name": "RivermaxMedia",
+			"Enabled": true
 		}
 	]
 }
\ No newline at end of file
