Change load order for ICVFX plugin to make sure ConsoleVariablesAsset is available.

#rnx

diff --git a/Engine/Plugins/VirtualProduction/ICVFX/ICVFX.uplugin b/Engine/Plugins/VirtualProduction/ICVFX/ICVFX.uplugin
--- a/Engine/Plugins/VirtualProduction/ICVFX/ICVFX.uplugin
+++ b/Engine/Plugins/VirtualProduction/ICVFX/ICVFX.uplugin
@@ -111,11 +111,11 @@
 			"Enabled": true
 		},
 		{
-			"Name": "VirtualProductionUtilities",
+			"Name": "ConsoleVariables",
 			"Enabled": true
 		},
 		{
-			"Name": "ConsoleVariables",
+			"Name": "VirtualProductionUtilities",
 			"Enabled": true
 		},
 		{
@@ -127,4 +127,4 @@
 			"Enabled": true
 		}
 	]
-}
\ No newline at end of file
+}
