Updating supported platforms, now matching QAVirtualProduction.

diff --git a/Engine/Plugins/VirtualProduction/ICVFX/ICVFX.uplugin b/Engine/Plugins/VirtualProduction/ICVFX/ICVFX.uplugin
--- a/Engine/Plugins/VirtualProduction/ICVFX/ICVFX.uplugin
+++ b/Engine/Plugins/VirtualProduction/ICVFX/ICVFX.uplugin
@@ -16,11 +16,18 @@
 	"Plugins": [
 		{
 			"Name": "AjaMedia",
-			"Enabled": true
+			"Enabled": true,
+			"SupportedTargetPlatforms": [
+				"Win64"
+			]
 		},
 		{
 			"Name": "BlackmagicMedia",
-			"Enabled": true
+			"Enabled": true,
+			"SupportedTargetPlatforms": [
+				"Win64",
+				"Linux"
+			]
 		},
 		{
 			"Name": "CameraCalibration",
@@ -48,7 +55,11 @@
 		},
 		{
 			"Name": "LiveLinkOverNDisplay",
-			"Enabled": true
+			"Enabled": true,
+			"SupportedTargetPlatforms": [
+				"Win64",
+				"Linux"
+			]
 		},
 		{
 			"Name": "LiveLinkLens",
@@ -72,7 +83,11 @@
 		},
 		{
 			"Name": "nDisplay",
-			"Enabled": true
+			"Enabled": true,
+			"SupportedTargetPlatforms": [
+				"Win64",
+				"Linux"
+			]
 		},
 		{
 			"Name": "OSC",
@@ -124,7 +139,10 @@
 		},
 		{
 			"Name": "RivermaxMedia",
-			"Enabled": true
+			"Enabled": true,
+			"SupportedTargetPlatforms": [
+				"Win64"
+			]
 		},
 		{
 			"Name": "ICVFXTesting",
